Live Streaming improvements
- New API for querying a list of existing live stream URLs for any game
- New "QueryLiveStreams" Blueprint node in the "Live Streaming" category (latent)
- Initial API support for sending and receiving online chat messages (no blueprint support yet)
- Fixed broadcasting console commands not unregistered when module is unloaded
- Various code clean-ups

diff --git a/Engine/Source/Editor/EditorLiveStreaming/EditorLiveStreaming.cpp b/Engine/Source/Editor/EditorLiveStreaming/EditorLiveStreaming.cpp
--- a/Engine/Source/Editor/EditorLiveStreaming/EditorLiveStreaming.cpp
+++ b/Engine/Source/Editor/EditorLiveStreaming/EditorLiveStreaming.cpp
@@ -80,15 +80,15 @@ void FEditorLiveStreaming::StartBroadcastingEditor()
 		// Select a live streaming service
 		{
 			static const FName LiveStreamingFeatureName( "LiveStreaming" );
-
-			// @todo livestream: Always using first live streaming plugin we find, but instead this should be a configurable option in the
-			// editor's preferences, based on the enumeration of available features from plugins
-			LiveStreamer = static_cast<ILiveStreamingService*>( IModularFeatures::Get().GetModularFeatureImplementation( LiveStreamingFeatureName, 0 ) );
+			LiveStreamer = &IModularFeatures::Get().GetModularFeature<ILiveStreamingService>( LiveStreamingFeatureName );
 		}
 
 		// Register to find out about status changes
 		LiveStreamer->OnStatusChanged().AddRaw( this, &FEditorLiveStreaming::BroadcastStatusCallback );
 
+		// @todo livestream: Allow connection to chat independently from broadcasting? (see removing delegate too)
+		LiveStreamer->OnChatMessage().AddRaw( this, &FEditorLiveStreaming::OnChatMessage );
+
 
 		// Tell our live streaming plugin to start broadcasting
 		{
@@ -196,6 +196,7 @@ void FEditorLiveStreaming::StopBroadcastingEditor()
 
 		// Unregister for status changes
 		LiveStreamer->OnStatusChanged().RemoveAll( this );
+		LiveStreamer->OnChatMessage().RemoveAll( this );
 
 		LiveStreamer = nullptr;
 
@@ -275,7 +276,9 @@ void FEditorLiveStreaming::BroadcastStatusCallback( const FLiveStreamingStatus&
 			( Status.StatusType != FLiveStreamingStatus::EStatusType::WebCamStarted &&
 			  Status.StatusType != FLiveStreamingStatus::EStatusType::WebCamStopped &&
 			  Status.StatusType != FLiveStreamingStatus::EStatusType::WebCamTextureNotReady &&
-			  Status.StatusType != FLiveStreamingStatus::EStatusType::WebCamTextureReady ) )
+			  Status.StatusType != FLiveStreamingStatus::EStatusType::WebCamTextureReady &&
+			  Status.StatusType != FLiveStreamingStatus::EStatusType::ChatConnected &&
+			  Status.StatusType != FLiveStreamingStatus::EStatusType::ChatDisconnected ) )
 		{
 			Notification->SetText( Status.CustomStatusDescription );
 		}
@@ -308,7 +311,8 @@ void FEditorLiveStreaming::BroadcastStatusCallback( const FLiveStreamingStatus&
 			State = SNotificationItem::CS_Fail;
 		}
 		else if( Status.StatusType == FLiveStreamingStatus::EStatusType::BroadcastStarted ||
-				 Status.StatusType == FLiveStreamingStatus::EStatusType::WebCamStarted )
+				 Status.StatusType == FLiveStreamingStatus::EStatusType::WebCamStarted ||
+				 Status.StatusType == FLiveStreamingStatus::EStatusType::ChatConnected )
 		{
 			State = SNotificationItem::CS_Success;
 		}
@@ -397,5 +401,10 @@ void FEditorLiveStreaming::CloseBroadcastStatusWindow()
 }
 
 
+void FEditorLiveStreaming::OnChatMessage( const FText& UserName, const FText& ChatMessage )
+{
+	// @todo livestream: Currently no chat UI is supported in the editor
+}
+
 
 #undef LOCTEXT_NAMESPACE
diff --git a/Engine/Source/Editor/EditorLiveStreaming/EditorLiveStreaming.h b/Engine/Source/Editor/EditorLiveStreaming/EditorLiveStreaming.h
--- a/Engine/Source/Editor/EditorLiveStreaming/EditorLiveStreaming.h
+++ b/Engine/Source/Editor/EditorLiveStreaming/EditorLiveStreaming.h
@@ -35,6 +35,9 @@ protected:
 	/** Closes the window that shows the current status of the broadcast */
 	void CloseBroadcastStatusWindow();
 
+	/** Called by the live streaming service when we get a new chat message to display */
+	void OnChatMessage( const FText& UserName, const FText& ChatMessage );
+
 
 private:
 
diff --git a/Engine/Source/Runtime/Engine/Public/Features/ILiveStreamingService.h b/Engine/Source/Runtime/Engine/Public/Features/ILiveStreamingService.h
--- a/Engine/Source/Runtime/Engine/Public/Features/ILiveStreamingService.h
+++ b/Engine/Source/Runtime/Engine/Public/Features/ILiveStreamingService.h
@@ -37,6 +37,12 @@ struct FLiveStreamingStatus
 
 		/** Web cam texture is no longer valid (GetWebCamTexture() will return null) */
 		WebCamTextureNotReady,
+
+		/** Connected to chat server.  You can now send messages! */
+		ChatConnected,
+
+		/** Disconnected from chat */
+		ChatDisconnected,
 	};
 
 	/** Status type */
@@ -129,12 +135,37 @@ struct FWebCamConfig
 };
 
 
+
+/**
+ * Contains information about an existing live streaming game.  Some services support querying for live broadcasts that
+ * are currently in progress
+ */
+struct FLiveStreamInfo
+{
+	/** Name of the game that is streaming */
+	FString GameName;
+
+	/** The title of the stream itself */
+	FString StreamName;
+
+	/** URL for the stream */
+	FString URL;
+
+
+	/** Default constructor */
+	FLiveStreamInfo()
+	{
+	}
+};
+
+
 /**
  * Generic interface to an implementation of a live streaming system, usually capable of capturing video and audio
  * and streaming live content to the internet
  */
 class ILiveStreamingService : public IModularFeature
 {
+	// @todo livestream: We should replace IDVRStreamingSystem with the new ILiveStreamingService, and update existing use cases (see FPS4PlatformFeaturesModule)
 
 public:
 
@@ -145,6 +176,10 @@ public:
 	virtual FOnStatusChanged& OnStatusChanged() = 0;
 
 
+	/**
+	 * Broadcasting
+	 */
+
 	/**
 	 * Called to start broadcasting video and audio.  Note that the actual streaming may not begin immediately.
 	 *
@@ -196,6 +231,11 @@ public:
 	 */
 	virtual void PushVideoFrame( const FColor* VideoFrameBuffer ) = 0;
 
+
+	/**
+	 * Web cam
+	 */
+
 	/**
 	 * Enables broadcasting from your web camera device
 	 *
@@ -221,4 +261,64 @@ public:
 	 * @return	Current web cam frame as a 2D texture
 	 */
 	virtual class UTexture2D* GetWebCamTexture() = 0;
+
+
+	/**
+	 * Chat
+	 */
+
+	/** Called when a chat message is received and should be displayed in the UI */
+	DECLARE_MULTICAST_DELEGATE_TwoParams( FOnChatMessage, const FText& /* UserName */, const FText& /* ChatMessage */ );
+
+	/** Event to register with to handle incoming chat messages when connected to a chat server */
+	virtual FOnChatMessage& OnChatMessage() = 0;
+
+	/**
+	 * Attempts to connect to the chat server for receiving and sending chat messages.  Make sure to bind a function to
+	 * the OnChatMessage() event so that you find out about incoming chat messages
+	 */
+	virtual void ConnectToChat() = 0;
+
+	/**
+	 * Disconnects from the chat server
+	 */
+	virtual void DisconnectFromChat() = 0;
+
+	/**
+	 * Checks to see if chat is currently enabled.  This doesn't necessarily mean that the user is connected to the chat
+	 * server at this time.  You need to listen for OnStatusChanged events to know for sure whether it is safe to
+	 * send chat messages at any given moment.  This may return true while a connection is still pending!
+	 * 
+	 * @return	True if the chat system is enabled and we're either already connected or a connection attempt is in progress
+	 */
+	virtual bool IsChatEnabled() const = 0;
+
+	/** 
+	 * Sends text to the chat channel, attributed to the currently logged in user.  In order for this to work, you must
+	 * have already successfully connected to a chat server.  You can listen for OnStatusChanged events to find out
+	 * when you are connected and it is possible to send chat messages
+	 *
+	 * @param	ChatMessage	The text to send
+	 */ 
+	virtual void SendChatMessage( const FString& ChatMessage ) = 0;
+
+
+	/**
+	 * Miscellaneous
+	 */
+
+	/** Completion callback for the QueryLiveStreams() function, that passes back the list of live streams along with a boolean
+	    variable that indicates whether the query was successful or not */
+	DECLARE_DELEGATE_TwoParams( FQueryLiveStreamsCallback, const TArray< FLiveStreamInfo >& /* LiveStreams */, bool /* bQueryWasSuccessful */ );
+
+	/**
+	 * Asynchronously queries a list of live streams that are currently in progress for the specified game name.  Note that some
+	 * services may not support this feature.  Your callback function will be called with the list of live streams that were found.
+	 * In the case there was an error or this feature wasn't supported, your callback will still be executed but will contain
+	 * a status variable that indicates failure
+	 *
+	 * @param	GameName				The name of the game to search for
+	 * @param	CompletionCallback		The function to invoke when the query has finished (or if an error occurs)
+	 */
+	virtual void QueryLiveStreams( const FString& GameName, FQueryLiveStreamsCallback CompletionCallback ) = 0;
 };
\ No newline at end of file
diff --git a/Engine/Source/Runtime/GameLiveStreaming/GameLiveStreaming.Build.cs b/Engine/Source/Runtime/GameLiveStreaming/GameLiveStreaming.Build.cs
--- a/Engine/Source/Runtime/GameLiveStreaming/GameLiveStreaming.Build.cs
+++ b/Engine/Source/Runtime/GameLiveStreaming/GameLiveStreaming.Build.cs
@@ -22,7 +22,8 @@ namespace UnrealBuildTool.Rules
 					"SlateCore",
 					"RenderCore",
 					"ShaderCore",
-					"RHI"
+					"RHI",
+					"OnlineSubsystemUtils"	// @todo livestream: We need to move the useful features into engine and eliminate this module
 				}
 			);
 		}
diff --git a/Engine/Source/Runtime/GameLiveStreaming/GameLiveStreaming.cpp b/Engine/Source/Runtime/GameLiveStreaming/GameLiveStreaming.cpp
--- a/Engine/Source/Runtime/GameLiveStreaming/GameLiveStreaming.cpp
+++ b/Engine/Source/Runtime/GameLiveStreaming/GameLiveStreaming.cpp
@@ -13,17 +13,7 @@ IMPLEMENT_MODULE( FGameLiveStreaming, GameLiveStreaming );
 #define LOCTEXT_NAMESPACE "GameLiveStreaming"
 
 FGameLiveStreaming::FGameLiveStreaming()
-{
-	LiveStreamer = nullptr;
-	ReadbackTextureIndex = 0;
-	ReadbackBufferIndex = 0;
-	ReadbackBuffers[0] = nullptr;
-	ReadbackBuffers[1] = nullptr;
-	bDrawSimpleWebCamVideo = true;
-
-	// Setup some console commands
-	{
-		static FAutoConsoleCommand BroadcastStartCommand
+	: BroadcastStartCommand
 		(
 			TEXT( "Broadcast.Start" ),
 			TEXT( "Starts broadcasting game video and audio using a live internet streaming service\n" )
@@ -54,9 +44,8 @@ FGameLiveStreaming::FGameLiveStreaming()
 					}
 					IGameLiveStreaming::Get().StartBroadcastingGame( Config );
 				} )
-		);
-
-		static FAutoConsoleCommand BroadcastStopCommand
+		),
+	  BroadcastStopCommand
 		(
 			TEXT( "Broadcast.Stop" ),
 			TEXT( "Stops broadcasting game video" ),
@@ -65,8 +54,15 @@ FGameLiveStreaming::FGameLiveStreaming()
 				{
 					IGameLiveStreaming::Get().StopBroadcastingGame();
 				} )
-		);
-	}
+		)
+{
+	bIsBroadcasting = false;
+	LiveStreamer = nullptr;
+	ReadbackTextureIndex = 0;
+	ReadbackBufferIndex = 0;
+	ReadbackBuffers[0] = nullptr;
+	ReadbackBuffers[1] = nullptr;
+	bDrawSimpleWebCamVideo = true;
 }
 
 
@@ -79,21 +75,24 @@ void FGameLiveStreaming::ShutdownModule()
 {
 	// Stop broadcasting at shutdown
 	StopBroadcastingGame();
-}
-
-
 
-bool FGameLiveStreaming::IsBroadcastingGame() const
-{
-	bool bAnythingBroadcasting = false;
 	if( LiveStreamer != nullptr )
 	{
-		if( LiveStreamer->IsBroadcasting() )
+		static const FName LiveStreamingFeatureName( "LiveStreaming" );
+		if( IModularFeatures::Get().IsModularFeatureAvailable( LiveStreamingFeatureName ) )	// Make sure the feature hasn't been destroyed before us
 		{
-			bAnythingBroadcasting = true;
+			LiveStreamer->OnStatusChanged().RemoveAll( this );
+			LiveStreamer->OnChatMessage().RemoveAll( this );
 		}
+		LiveStreamer = nullptr;
 	}
-	return bAnythingBroadcasting;
+}
+
+
+
+bool FGameLiveStreaming::IsBroadcastingGame() const
+{
+	return bIsBroadcasting && LiveStreamer != nullptr && LiveStreamer->IsBroadcasting();
 }
 
 
@@ -101,31 +100,25 @@ void FGameLiveStreaming::StartBroadcastingGame( const FGameBroadcastConfig& Game
 {
 	if( !IsBroadcastingGame() )
 	{
-		if( LiveStreamer != nullptr )
+		if( bIsBroadcasting )
 		{
 			FSlateRenderer* SlateRenderer = FSlateApplication::Get().GetRenderer().Get();
 			SlateRenderer->OnSlateWindowRendered().RemoveAll( this );
-			LiveStreamer->OnStatusChanged().RemoveAll( this );
-			LiveStreamer = nullptr;
+			bIsBroadcasting = false;
 		}
 
-		static const FName LiveStreamingFeatureName( "LiveStreaming" );
-		if( IModularFeatures::Get().IsModularFeatureAvailable( LiveStreamingFeatureName ) )
+		// We can GetLiveStreamingService() here to fill in our LiveStreamer variable lazily, to make sure the service plugin is loaded 
+		// before we try to cache it's interface pointer
+		if( GetLiveStreamingService() != nullptr  )
 		{
 			FSlateRenderer* SlateRenderer = FSlateApplication::Get().GetRenderer().Get();
 			SlateRenderer->OnSlateWindowRendered().AddRaw( this, &FGameLiveStreaming::OnSlateWindowRenderedDuringBroadcasting );
 
 			this->bDrawSimpleWebCamVideo = GameBroadcastConfig.bDrawSimpleWebCamVideo;
 
-			// Select a live streaming service
 			// @todo livestream: This will interfere with editor live streaming if both are running at the same time!  The editor live 
 			// streaming does check to make sure that game isn't already broadcasting, but the game currently doesn't have a good way to
 			// do that, besides asking the LiveStreamer itself.
-			// @todo livestream: Always using first live streaming plugin we find, but instead this should be a configurable option
-			LiveStreamer = static_cast<ILiveStreamingService*>( IModularFeatures::Get().GetModularFeatureImplementation( LiveStreamingFeatureName, 0 ) );
-
-			// Register to find out about status changes
-			LiveStreamer->OnStatusChanged().AddRaw( this, &FGameLiveStreaming::BroadcastStatusCallback );
 
 			UGameViewportClient* GameViewportClient = GEngine->GameViewport;
 			check( GameViewportClient != nullptr );
@@ -187,21 +180,21 @@ void FGameLiveStreaming::StartBroadcastingGame( const FGameBroadcastConfig& Game
 
 void FGameLiveStreaming::StopBroadcastingGame()
 {
-	if( LiveStreamer != nullptr )
+	if( bIsBroadcasting )
 	{
-		LiveStreamer->OnStatusChanged().RemoveAll( this );
-
-		if( LiveStreamer->IsBroadcasting() )
+		if( LiveStreamer != nullptr )
 		{
-			LiveStreamer->StopWebCam();
-			LiveStreamer->StopBroadcasting();
+			if( LiveStreamer->IsBroadcasting() )
+			{
+				LiveStreamer->StopWebCam();
+				LiveStreamer->StopBroadcasting();
+			}
 		}
 		if( FSlateApplication::IsInitialized() )	// During shutdown, Slate may have already been destroyed by the time our viewport gets cleaned up
 		{
 			FSlateRenderer* SlateRenderer = FSlateApplication::Get().GetRenderer().Get();
 			SlateRenderer->OnSlateWindowRendered().RemoveAll( this );
 		}
-		LiveStreamer = nullptr;
 
 		// Cleanup readback buffer textures
 		{
@@ -242,12 +235,8 @@ void FGameLiveStreaming::OnSlateWindowRenderedDuringBroadcasting( SWindow& Slate
 	}
 	else
 	{
-		FSlateRenderer* SlateRenderer = FSlateApplication::Get().GetRenderer().Get();
-
 		// No longer broadcasting.  The live streaming service may have been interrupted.
-		LiveStreamer->OnStatusChanged().RemoveAll( this );
-		LiveStreamer = nullptr;
-		SlateRenderer->OnSlateWindowRendered().RemoveAll( this );
+		StopBroadcastingGame();
 	}
 }
 
@@ -431,6 +420,26 @@ UTexture2D* FGameLiveStreaming::GetWebCamTexture()
 	return WebCamTexture;
 }
 
+
+ILiveStreamingService* FGameLiveStreaming::GetLiveStreamingService()
+{
+	if( LiveStreamer == nullptr )
+	{
+		static const FName LiveStreamingFeatureName( "LiveStreaming" );
+		if( IModularFeatures::Get().IsModularFeatureAvailable( LiveStreamingFeatureName ) )
+		{
+			// Select a live streaming service
+			LiveStreamer = &IModularFeatures::Get().GetModularFeature<ILiveStreamingService>( LiveStreamingFeatureName );
+
+			// Register to find out about status changes
+			LiveStreamer->OnStatusChanged().AddRaw( this, &FGameLiveStreaming::BroadcastStatusCallback );
+			LiveStreamer->OnChatMessage().AddRaw( this, &FGameLiveStreaming::OnChatMessage );
+		}
+	}
+	return LiveStreamer;
+}
+
+
 	
 void FGameLiveStreaming::BroadcastStatusCallback( const FLiveStreamingStatus& Status )
 {
@@ -438,5 +447,10 @@ void FGameLiveStreaming::BroadcastStatusCallback( const FLiveStreamingStatus& St
 }
 
 
+void FGameLiveStreaming::OnChatMessage( const FText& UserName, const FText& Text )
+{
+	// @todo livestream: Add support (and also, connecting, disconnecting, sending messages, etc.)
+}
+
 
 #undef LOCTEXT_NAMESPACE
diff --git a/Engine/Source/Runtime/GameLiveStreaming/GameLiveStreaming.h b/Engine/Source/Runtime/GameLiveStreaming/GameLiveStreaming.h
--- a/Engine/Source/Runtime/GameLiveStreaming/GameLiveStreaming.h
+++ b/Engine/Source/Runtime/GameLiveStreaming/GameLiveStreaming.h
@@ -25,12 +25,16 @@ public:
 	virtual void StopBroadcastingGame() override;
 	virtual void DrawSimpleWebCamVideo( UCanvas* Canvas ) override;
 	virtual class UTexture2D* GetWebCamTexture() override;
+	virtual class ILiveStreamingService* GetLiveStreamingService() override;
 
 protected:
 
 	/** Called by the live streaming service when streaming status has changed or an error has occurred */
 	void BroadcastStatusCallback( const struct FLiveStreamingStatus& Status );
 
+	/** Called by the live streaming service when a chat message is received */
+	void OnChatMessage( const FText& UserName, const FText& Text );
+
 	/**
 	 * Called by the Slate rendered on the game thread right after a window has been rendered, to allow us to
 	 * capture frame buffer content after everything has been rendered, including debug canvas graphics
@@ -49,6 +53,8 @@ protected:
 
 
 private:
+	/** Whether we're currently trying to broadcast */
+	bool bIsBroadcasting;
 
 	/** The live streaming service we're using.  Only valid while broadcasting. */
 	class ILiveStreamingService* LiveStreamer;
@@ -68,4 +74,9 @@ private:
 	/** True if we should draw a simple web cam video on top of the viewport while broadcasting */
 	bool bDrawSimpleWebCamVideo;
 
+	/** Console command for starting broadcast */
+	FAutoConsoleCommand BroadcastStartCommand;
+
+	/** Console command for stopping broadcast */
+	FAutoConsoleCommand BroadcastStopCommand;
 };
diff --git a/Engine/Source/Runtime/GameLiveStreaming/GameLiveStreamingFunctionLibrary.h b/Engine/Source/Runtime/GameLiveStreaming/GameLiveStreamingFunctionLibrary.h
new file mode 100644
--- /dev/null
+++ b/Engine/Source/Runtime/GameLiveStreaming/GameLiveStreamingFunctionLibrary.h
@@ -0,0 +1,47 @@
+// Copyright 1998-2014 Epic Games, Inc. All Rights Reserved.
+#pragma once
+#include "GameLiveStreamingFunctionLibrary.generated.h"
+
+UCLASS()
+class UGameLiveStreamingFunctionLibrary : public UBlueprintFunctionLibrary
+{
+	GENERATED_UCLASS_BODY()
+
+
+	/**
+	 * Checks to see if we are currently broadcasting live video (and possibly audio) from the game's viewport
+	 *
+	 * @return	True if we are currently transmitting
+	 */
+	UFUNCTION( BlueprintPure, Category="LiveStreaming" )
+	static bool IsBroadcastingGame();
+
+	/**
+	 * Starts broadcasting the game's video (and optionally audio) using an internet streaming service, if one is available 
+	 *
+	 * @param	FrameRate					Frame rate to stream video from when broadcasting to services like Twitch.
+	 * @param	ScreenScaling				How much to scale the broadcast video resolution down to reduce streaming bandwidth.  We recommend broadcasting at resolutions of 1280x720 or lower.  Some live streaming providers will not be able to transcode your video to a lower resolution, so using a high resolution stream may prevent low-bandwidth users from having a good viewing experience.
+	 * @param	bEnableWebCam				If enabled, video from your web camera will be captured and displayed while broadcasting, so that your viewers can see your presence.
+	 * @param	DesiredWebCamWidth			Desired web cam capture resolution width.  The web cam may only support a limited number of resolutions, so we'll choose one that matches as closely to this as possible
+	 * @param	DesiredWebCamHeight			Desired web cam capture resolution height.
+	 * @param	bCaptureAudioFromComputer	Enables broadcast of audio being played by your computer, such as in-game sounds
+	 * @param	bCaptureAudioFromMicrophone	Enables broadcast of audio from your default microphone recording device
+	 * @param	bDrawSimpleWebCamVideo		If enabled, the engine will draw a simple web cam image on top of the game viewport.  If you turn this off, it's up to you to draw the web cam image yourself.  You can access the web cam texture by calling IGameLiveStreaming::Get().GetWebCamTexture(). 
+	 */
+	UFUNCTION( BlueprintCallable, Category="LiveStreaming", meta=(FrameRate="30", ScreenScaling="1.0", bEnableWebCam="true", DesiredWebCamWidth="320", DesiredWebCamHeight="240", bCaptureAudioFromComputer="true", bCaptureAudioFromMicrophone="true", bDrawSimpleWebCamVideo="true") )
+	static void StartBroadcastingGame(
+		int32 FrameRate,
+		float ScreenScaling,
+		bool bEnableWebCam,
+		int32 DesiredWebCamWidth,
+		int32 DesiredWebCamHeight,
+		bool bCaptureAudioFromComputer,
+		bool bCaptureAudioFromMicrophone,
+		bool bDrawSimpleWebCamVideo );
+
+	/** Stops broadcasting the game */
+	UFUNCTION( BlueprintCallable, Category="LiveStreaming" )
+	static void StopBroadcastingGame();
+
+	// @todo livestream: Add Blueprint APIs for chat
+};
diff --git a/Engine/Source/Runtime/GameLiveStreaming/Public/GameLiveStreamingFunctionLibrary.h b/Engine/Source/Runtime/GameLiveStreaming/Public/GameLiveStreamingFunctionLibrary.h
deleted file mode 100644
diff --git a/Engine/Source/Runtime/GameLiveStreaming/Public/IGameLiveStreaming.h b/Engine/Source/Runtime/GameLiveStreaming/Public/IGameLiveStreaming.h
--- a/Engine/Source/Runtime/GameLiveStreaming/Public/IGameLiveStreaming.h
+++ b/Engine/Source/Runtime/GameLiveStreaming/Public/IGameLiveStreaming.h
@@ -93,6 +93,13 @@ public:
 	 */
 	virtual class UTexture2D* GetWebCamTexture() = 0;
 
+	/**
+	 * Gets direct access to the live streaming service we are using, or nullptr if we haven't selected one yet
+	 *
+	 * @return	Live streaming service interface
+	 */
+	virtual class ILiveStreamingService* GetLiveStreamingService() = 0;
+
 
 public:
 
diff --git a/Engine/Source/Runtime/GameLiveStreaming/QueryLiveStreamsCallbackProxy.cpp b/Engine/Source/Runtime/GameLiveStreaming/QueryLiveStreamsCallbackProxy.cpp
new file mode 100644
--- /dev/null
+++ b/Engine/Source/Runtime/GameLiveStreaming/QueryLiveStreamsCallbackProxy.cpp
@@ -0,0 +1,49 @@
+// Copyright 1998-2014 Epic Games, Inc. All Rights Reserved.
+#include "GameLiveStreamingModule.h"
+#include "QueryLiveStreamsCallbackProxy.h"
+#include "Runtime/Engine/Public/Features/ILiveStreamingService.h"
+
+
+UQueryLiveStreamsCallbackProxy::UQueryLiveStreamsCallbackProxy( const FPostConstructInitializeProperties& PCIP )
+	: Super(PCIP)
+{
+}
+
+
+UQueryLiveStreamsCallbackProxy* UQueryLiveStreamsCallbackProxy::QueryLiveStreams( const FString& GameName )
+{
+	UQueryLiveStreamsCallbackProxy* Proxy = NewObject<UQueryLiveStreamsCallbackProxy>();
+	Proxy->GameName = GameName;
+
+	return Proxy;
+}
+
+
+void UQueryLiveStreamsCallbackProxy::Activate()
+{
+	// Query the live streaming system to find out about which live streams are available for the game that was requested.  This
+	// will call back later on with results.
+	IGameLiveStreaming::Get().GetLiveStreamingService()->QueryLiveStreams(
+		this->GameName,
+		ILiveStreamingService::FQueryLiveStreamsCallback::CreateStatic( []( const TArray< FLiveStreamInfo >& LiveStreams, bool bQueryWasSuccessful, TWeakObjectPtr< UQueryLiveStreamsCallbackProxy > WeakThis )
+			{
+				// Make sure our UObject hasn't been destroyed at the time that the live streaming system calls back with results
+				UQueryLiveStreamsCallbackProxy* This = WeakThis.Get();
+				if( This != nullptr )
+				{
+					TArray< FBlueprintLiveStreamInfo > BlueprintLiveStreams;
+					for( const auto& LiveStream : LiveStreams )
+					{
+						FBlueprintLiveStreamInfo& BlueprintLiveStream = *new( BlueprintLiveStreams )FBlueprintLiveStreamInfo;
+						BlueprintLiveStream.GameName = LiveStream.GameName;
+						BlueprintLiveStream.StreamName = LiveStream.StreamName;
+						BlueprintLiveStream.URL = LiveStream.URL;
+					}
+
+					// We have results!
+					This->OnQueriedLiveStreams.Broadcast( BlueprintLiveStreams, bQueryWasSuccessful );
+				}
+			},
+			TWeakObjectPtr< UQueryLiveStreamsCallbackProxy>( this ) )
+		);
+}
diff --git a/Engine/Source/Runtime/GameLiveStreaming/QueryLiveStreamsCallbackProxy.h b/Engine/Source/Runtime/GameLiveStreaming/QueryLiveStreamsCallbackProxy.h
new file mode 100644
--- /dev/null
+++ b/Engine/Source/Runtime/GameLiveStreaming/QueryLiveStreamsCallbackProxy.h
@@ -0,0 +1,48 @@
+// Copyright 1998-2014 Epic Games, Inc. All Rights Reserved.
+#pragma once
+#include "OnlineBlueprintCallProxyBase.h"
+#include "QueryLiveStreamsCallbackProxy.generated.h"
+
+USTRUCT( BlueprintType )
+struct FBlueprintLiveStreamInfo 
+{
+	GENERATED_USTRUCT_BODY()
+
+	/** Name of the game that is streaming */
+	UPROPERTY( EditAnywhere, Category="LiveStreaming" )
+	FString GameName;
+
+	/** The title of the stream itself */
+	UPROPERTY( EditAnywhere, Category="LiveStreaming" )
+	FString StreamName;
+
+	/** URL for the stream */
+	UPROPERTY( EditAnywhere, Category="LiveStreaming" )
+	FString URL;
+};
+
+
+DECLARE_DYNAMIC_MULTICAST_DELEGATE_TwoParams( FOnQueryLiveStreamsCompleted, const TArray<FBlueprintLiveStreamInfo>&, LiveStreams, bool, bWasSuccessful );
+
+UCLASS()
+class UQueryLiveStreamsCallbackProxy : public UOnlineBlueprintCallProxyBase
+{
+	GENERATED_UCLASS_BODY()
+
+	// Called when the asynchronous request for live streams completes, passing along the list of live streams currently active, along with a boolean value that indicates whether the request was successful at all
+	UPROPERTY( BlueprintAssignable )
+	FOnQueryLiveStreamsCompleted OnQueriedLiveStreams;
+
+	/** Requests a list of live internet streams for the specified game name.  This will usually hit the internet so it could take a second or two. */
+	UFUNCTION( BlueprintCallable, Category="LiveStreaming", meta=( BlueprintInternalUseOnly="true" ) )
+	static UQueryLiveStreamsCallbackProxy* QueryLiveStreams( const FString& GameName );
+
+	/** UOnlineBlueprintCallProxyBase interface */
+	virtual void Activate() override;
+
+
+protected:
+
+	/** Name of the game that we're querying about */
+	FString GameName;
+};
\ No newline at end of file
