CIS fixes (non-unity)

diff --git a/Engine/Plugins/Script/ScriptEditorPlugin/Source/ScriptEditorPlugin/Private/ScriptEditorPlugin.cpp b/Engine/Plugins/Script/ScriptEditorPlugin/Source/ScriptEditorPlugin/Private/ScriptEditorPlugin.cpp
--- a/Engine/Plugins/Script/ScriptEditorPlugin/Source/ScriptEditorPlugin/Private/ScriptEditorPlugin.cpp
+++ b/Engine/Plugins/Script/ScriptEditorPlugin/Source/ScriptEditorPlugin/Private/ScriptEditorPlugin.cpp
@@ -3,6 +3,7 @@
 #include "ScriptEditorPluginPrivatePCH.h"
 #include "ScriptBlueprintCompiler.h"
 #include "KismetCompilerModule.h"
+#include "ScriptBlueprintEditor.h"
 
 DEFINE_LOG_CATEGORY(LogScriptEditorPlugin);
 
diff --git a/Engine/Source/Editor/UMGEditor/Private/WidgetBlueprintCompiler.cpp b/Engine/Source/Editor/UMGEditor/Private/WidgetBlueprintCompiler.cpp
--- a/Engine/Source/Editor/UMGEditor/Private/WidgetBlueprintCompiler.cpp
+++ b/Engine/Source/Editor/UMGEditor/Private/WidgetBlueprintCompiler.cpp
@@ -6,6 +6,7 @@
 #include "WidgetBlueprintCompiler.h"
 #include "Kismet2NameValidators.h"
 #include "KismetReinstanceUtilities.h"
+#include "MovieScene.h"
 
 template<>
 FString FNetNameMapping::MakeBaseName<UWidgetGraphNode_Base>(const UWidgetGraphNode_Base* Net)
diff --git a/Engine/Source/Runtime/GameLiveStreaming/QueryLiveStreamsCallbackProxy.cpp b/Engine/Source/Runtime/GameLiveStreaming/QueryLiveStreamsCallbackProxy.cpp
--- a/Engine/Source/Runtime/GameLiveStreaming/QueryLiveStreamsCallbackProxy.cpp
+++ b/Engine/Source/Runtime/GameLiveStreaming/QueryLiveStreamsCallbackProxy.cpp
@@ -2,7 +2,7 @@
 #include "GameLiveStreamingModule.h"
 #include "QueryLiveStreamsCallbackProxy.h"
 #include "Runtime/Engine/Public/Features/ILiveStreamingService.h"
-
+#include "IGameLiveStreaming.h"
 
 UQueryLiveStreamsCallbackProxy::UQueryLiveStreamsCallbackProxy( const FPostConstructInitializeProperties& PCIP )
 	: Super(PCIP)
