Fxied game broadcasting not starting successfully

diff --git a/Engine/Source/Runtime/GameLiveStreaming/GameLiveStreaming.cpp b/Engine/Source/Runtime/GameLiveStreaming/GameLiveStreaming.cpp
--- a/Engine/Source/Runtime/GameLiveStreaming/GameLiveStreaming.cpp
+++ b/Engine/Source/Runtime/GameLiveStreaming/GameLiveStreaming.cpp
@@ -173,6 +173,8 @@ void FGameLiveStreaming::StartBroadcastingGame( const FGameBroadcastConfig& Game
 				WebCamConfig.DesiredWebCamHeight = GameBroadcastConfig.DesiredWebCamHeight;
 				LiveStreamer->StartWebCam( WebCamConfig );
 			}
+
+			bIsBroadcasting = true;
 		}
 	}
 }
